Fix blank icons for Chaos and Geometry Collection plugins ChaosDestructionDemo.

diff --git a/Engine/Plugins/Experimental/ChaosSolverPlugin/Source/ChaosSolverEditor/Private/Chaos/ChaosSolverEditorStyle.h b/Engine/Plugins/Experimental/ChaosSolverPlugin/Source/ChaosSolverEditor/Private/Chaos/ChaosSolverEditorStyle.h
--- a/Engine/Plugins/Experimental/ChaosSolverPlugin/Source/ChaosSolverEditor/Private/Chaos/ChaosSolverEditorStyle.h
+++ b/Engine/Plugins/Experimental/ChaosSolverPlugin/Source/ChaosSolverEditor/Private/Chaos/ChaosSolverEditorStyle.h
@@ -13,11 +13,10 @@ public:
 		const FVector2D Icon64x64(64.f, 64.f);
 
 #if !IS_MONOLITHIC
-		FString PluginBasePath = FPaths::GetPath(FModuleManager::Get().GetModuleFilename("ChaosSolverEditor"));
-		SetContentRoot(PluginBasePath / TEXT("../../Resources"));
+		SetContentRoot(FPaths::EnginePluginsDir() / TEXT("Experimental/ChaosSolverPlugin/Resources"));
 #endif
 
-		Set("ClassIcon.ChaosSolver", new FSlateImageBrush(RootToContentDir(TEXT("ChaosSolver_16x.png")), Icon16x16));
+		Set("ClassIcon.ChaosSolver", new FSlateImageBrush(RootToCoreContentDir(TEXT("ChaosSolver_16x.png")), Icon16x16));
 		Set("ClassThumbnail.ChaosSolver", new FSlateImageBrush(RootToContentDir(TEXT("ChaosSolver_64x.png")), Icon64x64));
 
 		FSlateStyleRegistry::RegisterSlateStyle(*this);
diff --git a/Engine/Plugins/Experimental/FieldSystemPlugin/Source/FieldSyStemEditor/Private/Field/FieldSystemEditorStyle.h b/Engine/Plugins/Experimental/FieldSystemPlugin/Source/FieldSyStemEditor/Private/Field/FieldSystemEditorStyle.h
--- a/Engine/Plugins/Experimental/FieldSystemPlugin/Source/FieldSyStemEditor/Private/Field/FieldSystemEditorStyle.h
+++ b/Engine/Plugins/Experimental/FieldSystemPlugin/Source/FieldSyStemEditor/Private/Field/FieldSystemEditorStyle.h
@@ -13,8 +13,7 @@ public:
 		const FVector2D Icon64x64(64.f, 64.f);
 
 #if !IS_MONOLITHIC
-		FString PluginBasePath = FPaths::GetPath(FModuleManager::Get().GetModuleFilename("FieldSystemEditor"));
-		SetContentRoot(PluginBasePath / TEXT("../../Resources"));
+		SetContentRoot(FPaths::EnginePluginsDir() / TEXT("Experimental/FieldSystemPlugin/Resources"));
 #endif
 
 		Set("ClassIcon.FieldSystem", new FSlateImageBrush(RootToContentDir(TEXT("FieldSystem_16x.png")), Icon16x16));
diff --git a/Engine/Plugins/Experimental/GeometryCollectionPlugin/Source/GeometryCollectionEditor/Private/GeometryCollection/GeometryCollectionEditorStyle.h b/Engine/Plugins/Experimental/GeometryCollectionPlugin/Source/GeometryCollectionEditor/Private/GeometryCollection/GeometryCollectionEditorStyle.h
--- a/Engine/Plugins/Experimental/GeometryCollectionPlugin/Source/GeometryCollectionEditor/Private/GeometryCollection/GeometryCollectionEditorStyle.h
+++ b/Engine/Plugins/Experimental/GeometryCollectionPlugin/Source/GeometryCollectionEditor/Private/GeometryCollection/GeometryCollectionEditorStyle.h
@@ -13,8 +13,7 @@ public:
 		const FVector2D Icon64x64(64.f, 64.f);
 
 #if !IS_MONOLITHIC
-		FString PluginBasePath = FPaths::GetPath(FModuleManager::Get().GetModuleFilename("GeometryCollectionEditor"));
-		SetContentRoot(PluginBasePath / TEXT("../../Resources"));
+		SetContentRoot(FPaths::EnginePluginsDir() / TEXT("Experimental/GeometryCollectionPlugin/Resources"));
 #endif
 
 		Set("ClassIcon.GeometryCollection", new FSlateImageBrush(RootToContentDir(TEXT("GeometryCollection_16x.png")), Icon16x16));
