Fix ChaosSolver_16x.png warning.

diff --git a/Engine/Plugins/Experimental/ChaosSolverPlugin/Source/ChaosSolverEditor/Private/Chaos/ChaosSolverEditorStyle.h b/Engine/Plugins/Experimental/ChaosSolverPlugin/Source/ChaosSolverEditor/Private/Chaos/ChaosSolverEditorStyle.h
--- a/Engine/Plugins/Experimental/ChaosSolverPlugin/Source/ChaosSolverEditor/Private/Chaos/ChaosSolverEditorStyle.h
+++ b/Engine/Plugins/Experimental/ChaosSolverPlugin/Source/ChaosSolverEditor/Private/Chaos/ChaosSolverEditorStyle.h
@@ -16,7 +16,7 @@ public:
 		SetContentRoot(FPaths::EnginePluginsDir() / TEXT("Experimental/ChaosSolverPlugin/Resources"));
 #endif
 
-		Set("ClassIcon.ChaosSolver", new FSlateImageBrush(RootToCoreContentDir(TEXT("ChaosSolver_16x.png")), Icon16x16));
+		Set("ClassIcon.ChaosSolver", new FSlateImageBrush(RootToContentDir(TEXT("ChaosSolver_16x.png")), Icon16x16));
 		Set("ClassThumbnail.ChaosSolver", new FSlateImageBrush(RootToContentDir(TEXT("ChaosSolver_64x.png")), Icon64x64));
 
 		FSlateStyleRegistry::RegisterSlateStyle(*this);
